Make the oodle plugin and module used when producing standalone texture build workers configurable via INI.

diff --git a/Engine/Source/Programs/TextureBuildWorker/TextureBuildWorker.Build.cs b/Engine/Source/Programs/TextureBuildWorker/TextureBuildWorker.Build.cs
--- a/Engine/Source/Programs/TextureBuildWorker/TextureBuildWorker.Build.cs
+++ b/Engine/Source/Programs/TextureBuildWorker/TextureBuildWorker.Build.cs
@@ -1,6 +1,7 @@
 // Copyright Epic Games, Inc. All Rights Reserved.
 
 using System;
+using System.IO;
 using System.Linq;
 using EpicGames.Core;
 using UnrealBuildTool;
@@ -9,34 +10,50 @@ using UnrealBuildTool;
 [SupportedPlatforms(UnrealPlatformClass.Desktop)]
 public abstract class TextureBuildWorkerTarget : DerivedDataBuildWorkerTarget
 {
+	[ConfigFile(ConfigHierarchyType.Engine, "TextureBuildWorker", "ProjectOodlePlugin")]
+	public string ProjectOodlePlugin;
+
+	[ConfigFile(ConfigHierarchyType.Engine, "TextureBuildWorker", "ProjectOodleTextureFormatModule")]
+	public string ProjectOodleTextureFormatModule;
+
 	public TextureBuildWorkerTarget(TargetInfo Target) : base(Target)
 	{
 		SolutionDirectory += "/Texture";
 
-		var ProjectDesc = ProjectFile != null ? ProjectDescriptor.FromFile(ProjectFile) : null;
-
-		// Determine if TextureFormatOodle is enabled.
-		var TextureFormatOodleUPluginFile = new FileReference("../Plugins/Developer/TextureFormatOodle/TextureFormatOodle.uplugin");
-		var TextureFormatOodlePlugin = new PluginInfo(TextureFormatOodleUPluginFile, PluginType.Engine);
-
-		bool bTextureFormatOodlePluginEnabled =
-		Plugins.IsPluginEnabledForTarget(TextureFormatOodlePlugin, ProjectDesc, Target.Platform, Target.Configuration, TargetType.Program);
+		AddOodleModule(Target);
+	}
 
-		if (bTextureFormatOodlePluginEnabled)
-		{
-			ExtraModuleNames.Add("TextureFormatOodle");
-		}
-		else
+	private void AddOodleModule(TargetInfo Target)
+	{
+		FileReference OodleUPluginFile = new FileReference("../Plugins/Developer/TextureFormatOodle/TextureFormatOodle.uplugin");
+		PluginType OodlePluginType = PluginType.Engine;
+		string OodleTextureFormatModule = "TextureFormatOodle";
+		if (Target.ProjectFile != null)
 		{
-			// Check for a project specific Oodle plugin
-			foreach (PluginReferenceDescriptor PluginReference in ProjectDesc.Plugins)
+			if (!String.IsNullOrEmpty(ProjectOodlePlugin))
 			{
-				if (String.Compare(PluginReference.Name, "Oodle", true) == 0 && !PluginReference.bOptional)
+				if (!Path.IsPathRooted(ProjectOodlePlugin))
 				{
-					ExtraModuleNames.Add("OodleDXTTextureFormat");
-					break;
+					OodleUPluginFile = FileReference.FromString(Path.Combine(Target.ProjectFile.Directory.ToString(), ProjectOodlePlugin));
+					OodlePluginType = PluginType.Project;
 				}
 			}
+
+			if (!String.IsNullOrEmpty(ProjectOodleTextureFormatModule))
+			{
+				OodleTextureFormatModule = ProjectOodleTextureFormatModule;
+			}
+		}
+
+		// Determine if TextureFormatOodle is enabled.
+		var ProjectDesc = ProjectFile != null ? ProjectDescriptor.FromFile(ProjectFile) : null;
+		var OodlePlugin = new PluginInfo(OodleUPluginFile, OodlePluginType);
+		bool bOodlePluginEnabled =
+		Plugins.IsPluginEnabledForTarget(OodlePlugin, ProjectDesc, Target.Platform, Target.Configuration, TargetType.Program);
+
+		if (bOodlePluginEnabled)
+		{
+			ExtraModuleNames.Add(OodleTextureFormatModule);
 		}
 	}
 }
