fix TextureBuildWorker for ProjectOodlePlugin with Oodle Texture Version
may need both Oodle plugins

diff --git a/Engine/Source/Programs/TextureBuildWorker/TextureBuildWorker.Build.cs b/Engine/Source/Programs/TextureBuildWorker/TextureBuildWorker.Build.cs
--- a/Engine/Source/Programs/TextureBuildWorker/TextureBuildWorker.Build.cs
+++ b/Engine/Source/Programs/TextureBuildWorker/TextureBuildWorker.Build.cs
@@ -20,23 +20,31 @@ public abstract class TextureBuildWorkerTarget : DerivedDataBuildWorkerTarget
 	{
 		SolutionDirectory += "/Texture";
 
-		AddOodleModule(Target);
+		//always add the engine TextureFormatOodle (if enabled)
+		// then also add project Oodle (if enabled and configured)
+		// projects can use neither or both
+		AddOodleModule(Target,false);
+		AddOodleModule(Target,true);
 	}
 
-	private void AddOodleModule(TargetInfo Target)
+	private void AddOodleModule(TargetInfo Target,bool bProjectOodle)
 	{
 		FileReference OodleUPluginFile = new FileReference("../Plugins/Developer/TextureFormatOodle/TextureFormatOodle.uplugin");
 		PluginType OodlePluginType = PluginType.Engine;
 		string OodleTextureFormatModule = "TextureFormatOodle";
-		if (Target.ProjectFile != null)
+		if (bProjectOodle)
 		{
-			if (!String.IsNullOrEmpty(ProjectOodlePlugin))
+			if (Target.ProjectFile != null && !String.IsNullOrEmpty(ProjectOodlePlugin))
 			{
-				if (!Path.IsPathRooted(ProjectOodlePlugin))
-				{
-					OodleUPluginFile = FileReference.FromString(Path.Combine(Target.ProjectFile.Directory.ToString(), ProjectOodlePlugin));
-					OodlePluginType = PluginType.Project;
-				}
+				ProjectOodlePlugin = Path.Combine(Target.ProjectFile.Directory.ToString(), ProjectOodlePlugin);
+
+				OodleUPluginFile = FileReference.FromString(ProjectOodlePlugin);
+				OodlePluginType = PluginType.Project;
+			}
+			else
+			{
+				// bProjectOodle true but no project oodle set up
+				return;
 			}
 
 			if (!String.IsNullOrEmpty(ProjectOodleTextureFormatModule))
