[TextureBuildWorker]
* Fixed so oodle plugin parsing only happens once when creating module rules in UBT

diff --git a/Engine/Source/Programs/TextureBuildWorker/TextureBuildWorker.Build.cs b/Engine/Source/Programs/TextureBuildWorker/TextureBuildWorker.Build.cs
--- a/Engine/Source/Programs/TextureBuildWorker/TextureBuildWorker.Build.cs
+++ b/Engine/Source/Programs/TextureBuildWorker/TextureBuildWorker.Build.cs
@@ -29,6 +29,8 @@ public abstract class TextureBuildWorkerTarget : DerivedDataBuildWorkerTarget
 		AddOodleModule(Target,true);
 	}
 
+	static PluginInfo[] OodlePlugin = [null, null];
+
 	private void AddOodleModule(TargetInfo Target,bool bProjectOodle)
 	{
 		FileReference OodleUPluginFile = new FileReference("../Plugins/Developer/TextureFormatOodle/TextureFormatOodle.uplugin");
@@ -57,11 +59,16 @@ public abstract class TextureBuildWorkerTarget : DerivedDataBuildWorkerTarget
 
 		// Determine if TextureFormatOodle is enabled.
 		var ProjectDesc = ProjectFile != null ? ProjectDescriptor.FromFile(ProjectFile) : null;
-		var OodlePlugin = new PluginInfo(OodleUPluginFile, OodlePluginType);
-		bool bOodlePluginEnabled =
-		Plugins.IsPluginEnabledForTarget(OodlePlugin, ProjectDesc, Target.Platform, Target.Configuration, TargetType.Program);
 
-		if (bOodlePluginEnabled)
+		int TypeIndex = (int)OodlePluginType;
+
+		// We want to cache this and not re-parse plugin file over and over again
+		lock (OodlePlugin)
+		{
+			OodlePlugin[TypeIndex] ??= new PluginInfo(OodleUPluginFile, OodlePluginType);
+		}
+
+		if (Plugins.IsPluginEnabledForTarget(OodlePlugin[TypeIndex], ProjectDesc, Target.Platform, Target.Configuration, TargetType.Program))
 		{
 			ExtraModuleNames.Add(OodleTextureFormatModule);
 		}
