TextureBuildWorker.Build.cs no longer needs to look for Oodle Texture plugin

#Synced-CL 49850269

diff --git a/Engine/Source/Programs/TextureBuildWorker/TextureBuildWorker.Build.cs b/Engine/Source/Programs/TextureBuildWorker/TextureBuildWorker.Build.cs
--- a/Engine/Source/Programs/TextureBuildWorker/TextureBuildWorker.Build.cs
+++ b/Engine/Source/Programs/TextureBuildWorker/TextureBuildWorker.Build.cs
@@ -10,67 +10,13 @@ using UnrealBuildTool;
 [SupportedPlatforms(UnrealPlatformClass.Desktop)]
 public abstract class TextureBuildWorkerTarget : DerivedDataBuildWorkerTarget
 {
-	[ConfigFile(ConfigHierarchyType.Engine, "TextureBuildWorker", "ProjectOodlePlugin")]
-	public string ProjectOodlePlugin;
-
-	[ConfigFile(ConfigHierarchyType.Engine, "TextureBuildWorker", "ProjectOodleTextureFormatModule")]
-	public string ProjectOodleTextureFormatModule;
-
 	public TextureBuildWorkerTarget(TargetInfo Target) : base(Target)
 	{
 		IncludeOrderVersion = EngineIncludeOrderVersion.Latest;
 
 		SolutionDirectory += "/Texture";
-	
-		// engine TFO is a module now
-		// add optional project Oodle plugin
-		AddOodleModule(Target,true);
 	}
 
-	static PluginInfo[] OodlePlugin = [null, null];
-
-	private void AddOodleModule(TargetInfo Target,bool bProjectOodle)
-	{
-		FileReference OodleUPluginFile = new FileReference("../Plugins/Developer/TextureFormatOodle/TextureFormatOodle.uplugin");
-		PluginType OodlePluginType = PluginType.Engine;
-		string OodleTextureFormatModule = "TextureFormatOodle";
-		if (bProjectOodle)
-		{
-			if (Target.ProjectFile != null && !String.IsNullOrWhiteSpace(ProjectOodlePlugin))
-			{
-				ProjectOodlePlugin = Path.Combine(Target.ProjectFile.Directory.ToString(), ProjectOodlePlugin);
-
-				OodleUPluginFile = new(ProjectOodlePlugin);
-				OodlePluginType = PluginType.Project;
-			}
-			else
-			{
-				// bProjectOodle true but no project oodle set up
-				return;
-			}
-
-			if (!String.IsNullOrEmpty(ProjectOodleTextureFormatModule))
-			{
-				OodleTextureFormatModule = ProjectOodleTextureFormatModule;
-			}
-		}
-
-		// Determine if TextureFormatOodle is enabled.
-		var ProjectDesc = ProjectFile != null ? ProjectDescriptor.FromFile(ProjectFile) : null;
-
-		int TypeIndex = (int)OodlePluginType;
-
-		// We want to cache this and not re-parse plugin file over and over again
-		lock (OodlePlugin)
-		{
-			OodlePlugin[TypeIndex] ??= new PluginInfo(OodleUPluginFile, OodlePluginType);
-		}
-
-		if (Plugins.IsPluginEnabledForTarget(OodlePlugin[TypeIndex], ProjectDesc, Target.Platform, Target.Configuration, TargetType.Program))
-		{
-			ExtraModuleNames.Add(OodleTextureFormatModule);
-		}
-	}
 }
 
 public class TextureBuildWorker : ModuleRules
